[NO NON-PII COMMIT MESSAGE]

diff --git a/Templates/TP_FirstPerson/Config/DefaultInput.ini b/Templates/TP_FirstPerson/Config/DefaultInput.ini
--- a/Templates/TP_FirstPerson/Config/DefaultInput.ini
+++ b/Templates/TP_FirstPerson/Config/DefaultInput.ini
@@ -8,6 +8,7 @@
 +ActionMappings=(ActionName="Fire",Key=MotionController_Right_Trigger,bShift=False,bCtrl=False,bAlt=False,bCmd=False)
 
 +ActionMappings=(ActionName="ResetVR",Key=R,bShift=False,bCtrl=False,bAlt=False,bCmd=False)
++ActionMappings=(ActionName="ResetVR",Key=MotionController_Left_Grip1,bShift=False,bCtrl=False,bAlt=False,bCmd=False)
 
 +AxisMappings=(AxisName="MoveForward", Key=W, Scale=1.f)
 +AxisMappings=(AxisName="MoveForward", Key=S, Scale=-1.f)
diff --git a/Templates/TP_FirstPersonBP/Config/DefaultInput.ini b/Templates/TP_FirstPersonBP/Config/DefaultInput.ini
--- a/Templates/TP_FirstPersonBP/Config/DefaultInput.ini
+++ b/Templates/TP_FirstPersonBP/Config/DefaultInput.ini
@@ -43,6 +43,7 @@ bDefaultViewportMouseLock=True
 +ActionMappings=(ActionName="Fire",Key=MotionController_Right_Trigger,bShift=False,bCtrl=False,bAlt=False,bCmd=False)
 +ActionMappings=(ActionName="Jump",Key=MotionController_Left_Trigger,bShift=False,bCtrl=False,bAlt=False,bCmd=False)
 +ActionMappings=(ActionName="ResetVR",Key=R,bShift=False,bCtrl=False,bAlt=False,bCmd=False)
++ActionMappings=(ActionName="ResetVR",Key=MotionController_Left_Grip1,bShift=False,bCtrl=False,bAlt=False,bCmd=False)
 +AxisMappings=(AxisName="MoveForward",Key=W,Scale=1.000000)
 +AxisMappings=(AxisName="MoveForward",Key=S,Scale=-1.000000)
 +AxisMappings=(AxisName="MoveForward",Key=Up,Scale=1.000000)
diff --git a/Templates/TP_Flying/Config/DefaultInput.ini b/Templates/TP_Flying/Config/DefaultInput.ini
--- a/Templates/TP_Flying/Config/DefaultInput.ini
+++ b/Templates/TP_Flying/Config/DefaultInput.ini
@@ -27,6 +27,18 @@
 +AxisConfig=(AxisKeyName="MotionController_Right_Grip2Axis",AxisProperties=(DeadZone=0.000000,Sensitivity=1.000000,Exponent=1.000000,bInvert=False))
 +AxisConfig=(AxisKeyName="Gamepad_Special_Left_X",AxisProperties=(DeadZone=0.000000,Sensitivity=1.000000,Exponent=1.000000,bInvert=False))
 +AxisConfig=(AxisKeyName="Gamepad_Special_Left_Y",AxisProperties=(DeadZone=0.000000,Sensitivity=1.000000,Exponent=1.000000,bInvert=False))
++AxisConfig=(AxisKeyName="OculusTouch_Left_Thumbstick",AxisProperties=(DeadZone=0.000000,Sensitivity=1.000000,Exponent=1.000000,bInvert=False))
++AxisConfig=(AxisKeyName="OculusTouch_Left_FaceButton1",AxisProperties=(DeadZone=0.000000,Sensitivity=1.000000,Exponent=1.000000,bInvert=False))
++AxisConfig=(AxisKeyName="OculusTouch_Left_Trigger",AxisProperties=(DeadZone=0.000000,Sensitivity=1.000000,Exponent=1.000000,bInvert=False))
++AxisConfig=(AxisKeyName="OculusTouch_Left_FaceButton2",AxisProperties=(DeadZone=0.000000,Sensitivity=1.000000,Exponent=1.000000,bInvert=False))
++AxisConfig=(AxisKeyName="OculusTouch_Left_IndexPointing",AxisProperties=(DeadZone=0.000000,Sensitivity=1.000000,Exponent=1.000000,bInvert=False))
++AxisConfig=(AxisKeyName="OculusTouch_Left_ThumbUp",AxisProperties=(DeadZone=0.000000,Sensitivity=1.000000,Exponent=1.000000,bInvert=False))
++AxisConfig=(AxisKeyName="OculusTouch_Right_Thumbstick",AxisProperties=(DeadZone=0.000000,Sensitivity=1.000000,Exponent=1.000000,bInvert=False))
++AxisConfig=(AxisKeyName="OculusTouch_Right_FaceButton1",AxisProperties=(DeadZone=0.000000,Sensitivity=1.000000,Exponent=1.000000,bInvert=False))
++AxisConfig=(AxisKeyName="OculusTouch_Right_Trigger",AxisProperties=(DeadZone=0.000000,Sensitivity=1.000000,Exponent=1.000000,bInvert=False))
++AxisConfig=(AxisKeyName="OculusTouch_Right_FaceButton2",AxisProperties=(DeadZone=0.000000,Sensitivity=1.000000,Exponent=1.000000,bInvert=False))
++AxisConfig=(AxisKeyName="OculusTouch_Right_IndexPointing",AxisProperties=(DeadZone=0.000000,Sensitivity=1.000000,Exponent=1.000000,bInvert=False))
++AxisConfig=(AxisKeyName="OculusTouch_Right_ThumbUp",AxisProperties=(DeadZone=0.000000,Sensitivity=1.000000,Exponent=1.000000,bInvert=False))
 bAltEnterTogglesFullscreen=True
 bUseMouseForTouch=False
 bEnableMouseSmoothing=True
@@ -58,6 +70,12 @@ bDefaultViewportMouseLock=True
 +AxisMappings=(AxisName="MoveRight",Key=Gamepad_LeftX,Scale=1.000000)
 +AxisMappings=(AxisName="Thrust",Key=Gamepad_RightY,Scale=1.000000)
 +AxisMappings=(AxisName="Thrust",Key=SpaceBar,Scale=1.000000)
++AxisMappings=(AxisName="MoveUp",Key=MotionController_Left_Thumbstick_Y,Scale=1.000000)
++AxisMappings=(AxisName="MoveUp",Key=MotionController_Right_Thumbstick_Y,Scale=1.000000)
++AxisMappings=(AxisName="MoveRight",Key=MotionController_Left_Thumbstick_X,Scale=1.000000)
++AxisMappings=(AxisName="MoveRight",Key=MotionController_Right_Thumbstick_X,Scale=1.000000)
++AxisMappings=(AxisName="Thrust",Key=MotionController_Right_Trigger,Scale=1.000000)
++AxisMappings=(AxisName="Thrust",Key=MotionController_Left_Trigger,Scale=-1.000000)
 bAlwaysShowTouchInterface=False
 bShowConsoleOnFourFingerTap=True
 DefaultTouchInterface=/Engine/MobileResources/HUD/DefaultVirtualJoysticks.DefaultVirtualJoysticks
diff --git a/Templates/TP_Puzzle/Config/DefaultInput.ini b/Templates/TP_Puzzle/Config/DefaultInput.ini
--- a/Templates/TP_Puzzle/Config/DefaultInput.ini
+++ b/Templates/TP_Puzzle/Config/DefaultInput.ini
@@ -1,6 +1,50 @@
+
 [/Script/Engine.InputSettings]
-DefaultTouchInterface=
+-AxisConfig=(AxisKeyName="Gamepad_LeftX",AxisProperties=(DeadZone=0.25,Exponent=1.f,Sensitivity=1.f))
+-AxisConfig=(AxisKeyName="Gamepad_LeftY",AxisProperties=(DeadZone=0.25,Exponent=1.f,Sensitivity=1.f))
+-AxisConfig=(AxisKeyName="Gamepad_RightX",AxisProperties=(DeadZone=0.25,Exponent=1.f,Sensitivity=1.f))
+-AxisConfig=(AxisKeyName="Gamepad_RightY",AxisProperties=(DeadZone=0.25,Exponent=1.f,Sensitivity=1.f))
+-AxisConfig=(AxisKeyName="MouseX",AxisProperties=(DeadZone=0.f,Exponent=1.f,Sensitivity=0.07f))
+-AxisConfig=(AxisKeyName="MouseY",AxisProperties=(DeadZone=0.f,Exponent=1.f,Sensitivity=0.07f))
++AxisConfig=(AxisKeyName="Gamepad_LeftX",AxisProperties=(DeadZone=0.250000,Sensitivity=1.000000,Exponent=1.000000,bInvert=False))
++AxisConfig=(AxisKeyName="Gamepad_LeftY",AxisProperties=(DeadZone=0.250000,Sensitivity=1.000000,Exponent=1.000000,bInvert=False))
++AxisConfig=(AxisKeyName="Gamepad_RightX",AxisProperties=(DeadZone=0.250000,Sensitivity=1.000000,Exponent=1.000000,bInvert=False))
++AxisConfig=(AxisKeyName="Gamepad_RightY",AxisProperties=(DeadZone=0.250000,Sensitivity=1.000000,Exponent=1.000000,bInvert=False))
++AxisConfig=(AxisKeyName="MouseX",AxisProperties=(DeadZone=0.000000,Sensitivity=0.070000,Exponent=1.000000,bInvert=False))
++AxisConfig=(AxisKeyName="MouseY",AxisProperties=(DeadZone=0.000000,Sensitivity=0.070000,Exponent=1.000000,bInvert=False))
++AxisConfig=(AxisKeyName="MouseWheelAxis",AxisProperties=(DeadZone=0.000000,Sensitivity=1.000000,Exponent=1.000000,bInvert=False))
++AxisConfig=(AxisKeyName="Gamepad_LeftTriggerAxis",AxisProperties=(DeadZone=0.000000,Sensitivity=1.000000,Exponent=1.000000,bInvert=False))
++AxisConfig=(AxisKeyName="Gamepad_RightTriggerAxis",AxisProperties=(DeadZone=0.000000,Sensitivity=1.000000,Exponent=1.000000,bInvert=False))
++AxisConfig=(AxisKeyName="MotionController_Left_Thumbstick_X",AxisProperties=(DeadZone=0.000000,Sensitivity=1.000000,Exponent=1.000000,bInvert=False))
++AxisConfig=(AxisKeyName="MotionController_Left_Thumbstick_Y",AxisProperties=(DeadZone=0.000000,Sensitivity=1.000000,Exponent=1.000000,bInvert=False))
++AxisConfig=(AxisKeyName="MotionController_Left_TriggerAxis",AxisProperties=(DeadZone=0.000000,Sensitivity=1.000000,Exponent=1.000000,bInvert=False))
++AxisConfig=(AxisKeyName="MotionController_Left_Grip1Axis",AxisProperties=(DeadZone=0.000000,Sensitivity=1.000000,Exponent=1.000000,bInvert=False))
++AxisConfig=(AxisKeyName="MotionController_Left_Grip2Axis",AxisProperties=(DeadZone=0.000000,Sensitivity=1.000000,Exponent=1.000000,bInvert=False))
++AxisConfig=(AxisKeyName="MotionController_Right_Thumbstick_X",AxisProperties=(DeadZone=0.000000,Sensitivity=1.000000,Exponent=1.000000,bInvert=False))
++AxisConfig=(AxisKeyName="MotionController_Right_Thumbstick_Y",AxisProperties=(DeadZone=0.000000,Sensitivity=1.000000,Exponent=1.000000,bInvert=False))
++AxisConfig=(AxisKeyName="MotionController_Right_TriggerAxis",AxisProperties=(DeadZone=0.000000,Sensitivity=1.000000,Exponent=1.000000,bInvert=False))
++AxisConfig=(AxisKeyName="MotionController_Right_Grip1Axis",AxisProperties=(DeadZone=0.000000,Sensitivity=1.000000,Exponent=1.000000,bInvert=False))
++AxisConfig=(AxisKeyName="MotionController_Right_Grip2Axis",AxisProperties=(DeadZone=0.000000,Sensitivity=1.000000,Exponent=1.000000,bInvert=False))
++AxisConfig=(AxisKeyName="Gamepad_Special_Left_X",AxisProperties=(DeadZone=0.000000,Sensitivity=1.000000,Exponent=1.000000,bInvert=False))
++AxisConfig=(AxisKeyName="Gamepad_Special_Left_Y",AxisProperties=(DeadZone=0.000000,Sensitivity=1.000000,Exponent=1.000000,bInvert=False))
+bAltEnterTogglesFullscreen=True
+bUseMouseForTouch=False
+bEnableMouseSmoothing=True
+bEnableFOVScaling=True
+FOVScale=0.011110
+DoubleClickTime=0.200000
+bCaptureMouseOnLaunch=True
+DefaultViewportMouseCaptureMode=CapturePermanently_IncludingInitialMouseDown
+bDefaultViewportMouseLock=True
++ActionMappings=(ActionName="TriggerClick",Key=MotionController_Right_Trigger,bShift=False,bCtrl=False,bAlt=False,bCmd=False)
++ActionMappings=(ActionName="ResetVR",Key=R,bShift=False,bCtrl=False,bAlt=False,bCmd=False)
++ActionMappings=(ActionName="ResetVR",Key=MotionController_Left_Grip1,bShift=False,bCtrl=False,bAlt=False,bCmd=False)
++ActionMappings=(ActionName="TriggerClick",Key=Gamepad_RightTrigger,bShift=False,bCtrl=False,bAlt=False,bCmd=False)
+bAlwaysShowTouchInterface=False
+bShowConsoleOnFourFingerTap=True
+DefaultTouchInterface=None
+ConsoleKey=None
+-ConsoleKeys=Tilde
++ConsoleKeys=Tilde
 
-!AxisMappings=Empty
-!ActionMappings=Empty
 
diff --git a/Templates/TP_Puzzle/Source/TP_Puzzle/TP_Puzzle.Build.cs b/Templates/TP_Puzzle/Source/TP_Puzzle/TP_Puzzle.Build.cs
--- a/Templates/TP_Puzzle/Source/TP_Puzzle/TP_Puzzle.Build.cs
+++ b/Templates/TP_Puzzle/Source/TP_Puzzle/TP_Puzzle.Build.cs
@@ -6,6 +6,6 @@ public class TP_Puzzle : ModuleRules
 {
 	public TP_Puzzle(TargetInfo Target)
 	{
-		PublicDependencyModuleNames.AddRange(new string[] { "Core", "CoreUObject", "Engine", "InputCore" });
+		PublicDependencyModuleNames.AddRange(new string[] { "Core", "CoreUObject", "Engine", "InputCore", "HeadMountedDisplay" });
 	}
 }
diff --git a/Templates/TP_Puzzle/Source/TP_Puzzle/TP_PuzzleBlock.cpp b/Templates/TP_Puzzle/Source/TP_Puzzle/TP_PuzzleBlock.cpp
--- a/Templates/TP_Puzzle/Source/TP_Puzzle/TP_PuzzleBlock.cpp
+++ b/Templates/TP_Puzzle/Source/TP_Puzzle/TP_PuzzleBlock.cpp
@@ -10,10 +10,12 @@ ATP_PuzzleBlock::ATP_PuzzleBlock()
 	struct FConstructorStatics
 	{
 		ConstructorHelpers::FObjectFinderOptional<UStaticMesh> PlaneMesh;
+		ConstructorHelpers::FObjectFinderOptional<UMaterial> BaseMaterial;
 		ConstructorHelpers::FObjectFinderOptional<UMaterialInstance> BlueMaterial;
 		ConstructorHelpers::FObjectFinderOptional<UMaterialInstance> OrangeMaterial;
 		FConstructorStatics()
 			: PlaneMesh(TEXT("/Game/Puzzle/Meshes/PuzzleCube.PuzzleCube"))
+			, BaseMaterial(TEXT("/Game/Puzzle/Meshes/BaseMaterial.BaseMaterial"))
 			, BlueMaterial(TEXT("/Game/Puzzle/Meshes/BlueMaterial.BlueMaterial"))
 			, OrangeMaterial(TEXT("/Game/Puzzle/Meshes/OrangeMaterial.OrangeMaterial"))
 		{
@@ -36,13 +38,26 @@ ATP_PuzzleBlock::ATP_PuzzleBlock()
 	BlockMesh->OnInputTouchBegin.AddDynamic(this, &ATP_PuzzleBlock::OnFingerPressedBlock);
 
 	// Save a pointer to the orange material
+	BaseMaterial = ConstructorStatics.BaseMaterial.Get();
+	BlueMaterial = ConstructorStatics.BlueMaterial.Get();
 	OrangeMaterial = ConstructorStatics.OrangeMaterial.Get();
 }
 
 void ATP_PuzzleBlock::BlockClicked(UPrimitiveComponent* ClickedComp, FKey ButtonClicked)
+{
+	HandleClicked();
+}
+
+
+void ATP_PuzzleBlock::OnFingerPressedBlock(ETouchIndex::Type FingerIndex, UPrimitiveComponent* TouchedComponent)
+{
+	HandleClicked();
+}
+
+void ATP_PuzzleBlock::HandleClicked()
 {
 	// Check we are not already active
-	if(!bIsActive)
+	if (!bIsActive)
 	{
 		bIsActive = true;
 
@@ -50,15 +65,27 @@ void ATP_PuzzleBlock::BlockClicked(UPrimitiveComponent* ClickedComp, FKey Button
 		BlockMesh->SetMaterial(0, OrangeMaterial);
 
 		// Tell the Grid
-		if(OwningGrid != NULL)
+		if (OwningGrid != nullptr)
 		{
 			OwningGrid->AddScore();
 		}
 	}
 }
 
-
-void ATP_PuzzleBlock::OnFingerPressedBlock(ETouchIndex::Type FingerIndex, UPrimitiveComponent* TouchedComponent)
+void ATP_PuzzleBlock::Highlight(bool bOn)
 {
-	BlockClicked(TouchedComponent, EKeys::Invalid);
+	// Do not highlight if the block has already been activated.
+	if (bIsActive)
+	{
+		return;
+	}
+
+	if (bOn)
+	{
+		BlockMesh->SetMaterial(0, BaseMaterial);
+	}
+	else
+	{
+		BlockMesh->SetMaterial(0, BlueMaterial);
+	}
 }
diff --git a/Templates/TP_Puzzle/Source/TP_Puzzle/TP_PuzzleBlock.h b/Templates/TP_Puzzle/Source/TP_Puzzle/TP_PuzzleBlock.h
--- a/Templates/TP_Puzzle/Source/TP_Puzzle/TP_PuzzleBlock.h
+++ b/Templates/TP_Puzzle/Source/TP_Puzzle/TP_PuzzleBlock.h
@@ -23,6 +23,14 @@ public:
 	/** Are we currently active? */
 	bool bIsActive;
 
+	/** Pointer to white material used on the focused block */
+	UPROPERTY()
+	class UMaterial* BaseMaterial;
+
+	/** Pointer to blue material used on inactive blocks */
+	UPROPERTY()
+	class UMaterialInstance* BlueMaterial;
+
 	/** Pointer to orange material used on active blocks */
 	UPROPERTY()
 	class UMaterialInstance* OrangeMaterial;
@@ -39,6 +47,10 @@ public:
 	UFUNCTION()
 	void OnFingerPressedBlock(ETouchIndex::Type FingerIndex, UPrimitiveComponent* TouchedComponent);
 
+	void HandleClicked();
+
+	void Highlight(bool bOn);
+
 public:
 	/** Returns DummyRoot subobject **/
 	FORCEINLINE class USceneComponent* GetDummyRoot() const { return DummyRoot; }
diff --git a/Templates/TP_Puzzle/Source/TP_Puzzle/TP_PuzzleBlockGrid.cpp b/Templates/TP_Puzzle/Source/TP_Puzzle/TP_PuzzleBlockGrid.cpp
--- a/Templates/TP_Puzzle/Source/TP_Puzzle/TP_PuzzleBlockGrid.cpp
+++ b/Templates/TP_Puzzle/Source/TP_Puzzle/TP_PuzzleBlockGrid.cpp
@@ -46,7 +46,7 @@ void ATP_PuzzleBlockGrid::BeginPlay()
 		ATP_PuzzleBlock* NewBlock = GetWorld()->SpawnActor<ATP_PuzzleBlock>(BlockLocation, FRotator(0,0,0));
 
 		// Tell the block about its owner
-		if(NewBlock != NULL)
+		if (NewBlock != nullptr)
 		{
 			NewBlock->OwningGrid = this;
 		}
diff --git a/Templates/TP_Puzzle/Source/TP_Puzzle/TP_PuzzleGameMode.cpp b/Templates/TP_Puzzle/Source/TP_Puzzle/TP_PuzzleGameMode.cpp
--- a/Templates/TP_Puzzle/Source/TP_Puzzle/TP_PuzzleGameMode.cpp
+++ b/Templates/TP_Puzzle/Source/TP_Puzzle/TP_PuzzleGameMode.cpp
@@ -3,11 +3,12 @@
 #include "TP_Puzzle.h"
 #include "TP_PuzzleGameMode.h"
 #include "TP_PuzzlePlayerController.h"
+#include "TP_PuzzlePawn.h"
 
 ATP_PuzzleGameMode::ATP_PuzzleGameMode()
 {
 	// no pawn by default
-	DefaultPawnClass = NULL;
+	DefaultPawnClass = ATP_PuzzlePawn::StaticClass();
 	// use our own player controller class
 	PlayerControllerClass = ATP_PuzzlePlayerController::StaticClass();
 }
diff --git a/Templates/TP_Puzzle/Source/TP_Puzzle/TP_PuzzlePawn.cpp b/Templates/TP_Puzzle/Source/TP_Puzzle/TP_PuzzlePawn.cpp
new file mode 100644
--- /dev/null
+++ b/Templates/TP_Puzzle/Source/TP_Puzzle/TP_PuzzlePawn.cpp
@@ -0,0 +1,98 @@
+// Copyright 1998-2015 Epic Games, Inc. All Rights Reserved.
+
+#include "TP_Puzzle.h"
+#include "TP_PuzzlePawn.h"
+#include "TP_PuzzleBlock.h"
+#include "Kismet/HeadMountedDisplayFunctionLibrary.h"
+#include "DrawDebugHelpers.h"
+
+ATP_PuzzlePawn::ATP_PuzzlePawn(const FObjectInitializer& ObjectInitializer) 
+	: Super(ObjectInitializer)
+{
+	AutoPossessPlayer = EAutoReceiveInput::Player0;
+}
+
+void ATP_PuzzlePawn::Tick(float DeltaSeconds)
+{
+	Super::Tick(DeltaSeconds);
+
+	if (APlayerController* PC = Cast<APlayerController>(GetController()))
+	{
+		if (UHeadMountedDisplayFunctionLibrary::IsHeadMountedDisplayEnabled())
+		{
+			if (UCameraComponent* OurCamera = PC->GetViewTarget()->FindComponentByClass<UCameraComponent>())
+			{
+				FVector Start = OurCamera->GetComponentLocation();
+				FVector End = Start + (OurCamera->GetComponentRotation().Vector() * 8000.0f);
+				TraceForBlock(Start, End, true);
+			}
+		}
+		else
+		{
+			FVector Start, Dir, End;
+			PC->DeprojectMousePositionToWorld(Start, Dir);
+			End = Start + (Dir * 8000.0f);
+			TraceForBlock(Start, End, false);
+		}
+	}
+}
+
+void ATP_PuzzlePawn::SetupPlayerInputComponent(UInputComponent* PlayerInputComponent)
+{
+	Super::SetupPlayerInputComponent(PlayerInputComponent);
+
+	PlayerInputComponent->BindAction("OnResetVR", EInputEvent::IE_Pressed, this, &ATP_PuzzlePawn::OnResetVR);
+	PlayerInputComponent->BindAction("TriggerClick", EInputEvent::IE_Pressed, this, &ATP_PuzzlePawn::TriggerClick);
+}
+
+void ATP_PuzzlePawn::CalcCamera(float DeltaTime, struct FMinimalViewInfo& OutResult)
+{
+	Super::CalcCamera(DeltaTime, OutResult);
+
+	OutResult.Rotation = FRotator(-90.0f, -90.0f, 0.0f);
+}
+
+void ATP_PuzzlePawn::OnResetVR()
+{
+	UHeadMountedDisplayFunctionLibrary::ResetOrientationAndPosition();
+}
+
+void ATP_PuzzlePawn::TriggerClick()
+{
+	if (CurrentBlockFocus)
+	{
+		CurrentBlockFocus->HandleClicked();
+	}
+}
+
+void ATP_PuzzlePawn::TraceForBlock(const FVector& Start, const FVector& End, bool bDrawDebugHelpers)
+{
+	FHitResult HitResult;
+	GetWorld()->LineTraceSingleByChannel(HitResult, Start, End, ECC_Visibility);
+	if (bDrawDebugHelpers)
+	{
+		DrawDebugLine(GetWorld(), Start, HitResult.Location, FColor::Red);
+		DrawDebugSolidBox(GetWorld(), HitResult.Location, FVector(20.0f), FColor::Red);
+	}
+	if (HitResult.Actor.IsValid())
+	{
+		ATP_PuzzleBlock* HitBlock = Cast<ATP_PuzzleBlock>(HitResult.Actor.Get());
+		if (CurrentBlockFocus != HitBlock)
+		{
+			if (CurrentBlockFocus)
+			{
+				CurrentBlockFocus->Highlight(false);
+			}
+			if (HitBlock)
+			{
+				HitBlock->Highlight(true);
+			}
+			CurrentBlockFocus = HitBlock;
+		}
+	}
+	else if (CurrentBlockFocus)
+	{
+		CurrentBlockFocus->Highlight(false);
+		CurrentBlockFocus = nullptr;
+	}
+}
\ No newline at end of file
diff --git a/Templates/TP_Puzzle/Source/TP_Puzzle/TP_PuzzlePawn.h b/Templates/TP_Puzzle/Source/TP_Puzzle/TP_PuzzlePawn.h
new file mode 100644
--- /dev/null
+++ b/Templates/TP_Puzzle/Source/TP_Puzzle/TP_PuzzlePawn.h
@@ -0,0 +1,26 @@
+// Copyright 1998-2015 Epic Games, Inc. All Rights Reserved.
+#pragma once
+
+#include "TP_PuzzlePawn.generated.h"
+
+UCLASS(config=Game)
+class ATP_PuzzlePawn : public APawn
+{
+	GENERATED_UCLASS_BODY()
+
+public:
+
+	virtual void Tick(float DeltaSeconds) override;
+
+	virtual void SetupPlayerInputComponent(UInputComponent* PlayerInputComponent) override;
+
+	virtual void CalcCamera(float DeltaTime, struct FMinimalViewInfo& OutResult) override;
+
+protected:
+	void OnResetVR();
+	void TriggerClick();
+	void TraceForBlock(const FVector& Start, const FVector& End, bool bDrawDebugHelpers);
+
+	UPROPERTY(EditInstanceOnly, BlueprintReadWrite)
+	class ATP_PuzzleBlock* CurrentBlockFocus;
+};
diff --git a/Templates/TP_ThirdPerson/Config/DefaultInput.ini b/Templates/TP_ThirdPerson/Config/DefaultInput.ini
--- a/Templates/TP_ThirdPerson/Config/DefaultInput.ini
+++ b/Templates/TP_ThirdPerson/Config/DefaultInput.ini
@@ -21,3 +21,4 @@
 +AxisMappings=(AxisName="LookUp", Key=MouseY, Scale=-1.f)
 
 +ActionMappings=(ActionName="ResetVR", Key=R)
++ActionMappings=(ActionName="ResetVR",Key=MotionController_Left_Grip1,bShift=False,bCtrl=False,bAlt=False,bCmd=False)
diff --git a/Templates/TP_ThirdPersonBP/Config/DefaultInput.ini b/Templates/TP_ThirdPersonBP/Config/DefaultInput.ini
--- a/Templates/TP_ThirdPersonBP/Config/DefaultInput.ini
+++ b/Templates/TP_ThirdPersonBP/Config/DefaultInput.ini
@@ -40,6 +40,7 @@ bDefaultViewportMouseLock=True
 +ActionMappings=(ActionName="Jump",Key=Gamepad_FaceButton_Bottom,bShift=False,bCtrl=False,bAlt=False,bCmd=False)
 +ActionMappings=(ActionName="Jump",Key=MotionController_Left_Trigger,bShift=False,bCtrl=False,bAlt=False,bCmd=False)
 +ActionMappings=(ActionName="ResetVR",Key=R,bShift=False,bCtrl=False,bAlt=False,bCmd=False)
++ActionMappings=(ActionName="ResetVR",Key=MotionController_Left_Grip1,bShift=False,bCtrl=False,bAlt=False,bCmd=False)
 +ActionMappings=(ActionName="Jump",Key=MotionController_Right_Trigger,bShift=False,bCtrl=False,bAlt=False,bCmd=False)
 +AxisMappings=(AxisName="MoveForward",Key=W,Scale=1.000000)
 +AxisMappings=(AxisName="MoveForward",Key=S,Scale=-1.000000)
diff --git a/Templates/TP_TopDown/Config/DefaultInput.ini b/Templates/TP_TopDown/Config/DefaultInput.ini
--- a/Templates/TP_TopDown/Config/DefaultInput.ini
+++ b/Templates/TP_TopDown/Config/DefaultInput.ini
@@ -1,14 +1,56 @@
 [/Script/Engine.InputSettings]
-DefaultTouchInterface=
+-AxisConfig=(AxisKeyName="Gamepad_LeftX",AxisProperties=(DeadZone=0.25,Exponent=1.f,Sensitivity=1.f))
+-AxisConfig=(AxisKeyName="Gamepad_LeftY",AxisProperties=(DeadZone=0.25,Exponent=1.f,Sensitivity=1.f))
+-AxisConfig=(AxisKeyName="Gamepad_RightX",AxisProperties=(DeadZone=0.25,Exponent=1.f,Sensitivity=1.f))
+-AxisConfig=(AxisKeyName="Gamepad_RightY",AxisProperties=(DeadZone=0.25,Exponent=1.f,Sensitivity=1.f))
+-AxisConfig=(AxisKeyName="MouseX",AxisProperties=(DeadZone=0.f,Exponent=1.f,Sensitivity=0.07f))
+-AxisConfig=(AxisKeyName="MouseY",AxisProperties=(DeadZone=0.f,Exponent=1.f,Sensitivity=0.07f))
++AxisConfig=(AxisKeyName="Gamepad_LeftX",AxisProperties=(DeadZone=0.250000,Sensitivity=1.000000,Exponent=1.000000,bInvert=False))
++AxisConfig=(AxisKeyName="Gamepad_LeftY",AxisProperties=(DeadZone=0.250000,Sensitivity=1.000000,Exponent=1.000000,bInvert=False))
++AxisConfig=(AxisKeyName="Gamepad_RightX",AxisProperties=(DeadZone=0.250000,Sensitivity=1.000000,Exponent=1.000000,bInvert=False))
++AxisConfig=(AxisKeyName="Gamepad_RightY",AxisProperties=(DeadZone=0.250000,Sensitivity=1.000000,Exponent=1.000000,bInvert=False))
++AxisConfig=(AxisKeyName="MouseX",AxisProperties=(DeadZone=0.000000,Sensitivity=0.070000,Exponent=1.000000,bInvert=False))
++AxisConfig=(AxisKeyName="MouseY",AxisProperties=(DeadZone=0.000000,Sensitivity=0.070000,Exponent=1.000000,bInvert=False))
++AxisConfig=(AxisKeyName="MouseWheelAxis",AxisProperties=(DeadZone=0.000000,Sensitivity=1.000000,Exponent=1.000000,bInvert=False))
++AxisConfig=(AxisKeyName="Gamepad_LeftTriggerAxis",AxisProperties=(DeadZone=0.000000,Sensitivity=1.000000,Exponent=1.000000,bInvert=False))
++AxisConfig=(AxisKeyName="Gamepad_RightTriggerAxis",AxisProperties=(DeadZone=0.000000,Sensitivity=1.000000,Exponent=1.000000,bInvert=False))
++AxisConfig=(AxisKeyName="MotionController_Left_Thumbstick_X",AxisProperties=(DeadZone=0.000000,Sensitivity=1.000000,Exponent=1.000000,bInvert=False))
++AxisConfig=(AxisKeyName="MotionController_Left_Thumbstick_Y",AxisProperties=(DeadZone=0.000000,Sensitivity=1.000000,Exponent=1.000000,bInvert=False))
++AxisConfig=(AxisKeyName="MotionController_Left_TriggerAxis",AxisProperties=(DeadZone=0.000000,Sensitivity=1.000000,Exponent=1.000000,bInvert=False))
++AxisConfig=(AxisKeyName="MotionController_Left_Grip1Axis",AxisProperties=(DeadZone=0.000000,Sensitivity=1.000000,Exponent=1.000000,bInvert=False))
++AxisConfig=(AxisKeyName="MotionController_Left_Grip2Axis",AxisProperties=(DeadZone=0.000000,Sensitivity=1.000000,Exponent=1.000000,bInvert=False))
++AxisConfig=(AxisKeyName="MotionController_Right_Thumbstick_X",AxisProperties=(DeadZone=0.000000,Sensitivity=1.000000,Exponent=1.000000,bInvert=False))
++AxisConfig=(AxisKeyName="MotionController_Right_Thumbstick_Y",AxisProperties=(DeadZone=0.000000,Sensitivity=1.000000,Exponent=1.000000,bInvert=False))
++AxisConfig=(AxisKeyName="MotionController_Right_TriggerAxis",AxisProperties=(DeadZone=0.000000,Sensitivity=1.000000,Exponent=1.000000,bInvert=False))
++AxisConfig=(AxisKeyName="MotionController_Right_Grip1Axis",AxisProperties=(DeadZone=0.000000,Sensitivity=1.000000,Exponent=1.000000,bInvert=False))
++AxisConfig=(AxisKeyName="MotionController_Right_Grip2Axis",AxisProperties=(DeadZone=0.000000,Sensitivity=1.000000,Exponent=1.000000,bInvert=False))
++AxisConfig=(AxisKeyName="Gamepad_Special_Left_X",AxisProperties=(DeadZone=0.000000,Sensitivity=1.000000,Exponent=1.000000,bInvert=False))
++AxisConfig=(AxisKeyName="Gamepad_Special_Left_Y",AxisProperties=(DeadZone=0.000000,Sensitivity=1.000000,Exponent=1.000000,bInvert=False))
+bAltEnterTogglesFullscreen=True
+bUseMouseForTouch=False
+bEnableMouseSmoothing=True
+bEnableFOVScaling=True
+FOVScale=0.011110
+DoubleClickTime=0.200000
+bCaptureMouseOnLaunch=True
+DefaultViewportMouseCaptureMode=CapturePermanently_IncludingInitialMouseDown
+bDefaultViewportMouseLock=True
++ActionMappings=(ActionName="SetDestination",Key=LeftMouseButton,bShift=False,bCtrl=False,bAlt=False,bCmd=False)
++ActionMappings=(ActionName="SetDestination",Key=MotionController_Right_Trigger,bShift=False,bCtrl=False,bAlt=False,bCmd=False)
++ActionMappings=(ActionName="SetDestination",Key=Gamepad_RightTrigger,bShift=False,bCtrl=False,bAlt=False,bCmd=False)
++ActionMappings=(ActionName="ResetVR",Key=R,bShift=False,bCtrl=False,bAlt=False,bCmd=False)
++ActionMappings=(ActionName="ResetVR",Key=MotionController_Left_Grip1,bShift=False,bCtrl=False,bAlt=False,bCmd=False)
++AxisMappings=(AxisName="MoveForward",Key=W,Scale=1.000000)
++AxisMappings=(AxisName="MoveForward",Key=S,Scale=-1.000000)
++AxisMappings=(AxisName="MoveRight",Key=D,Scale=1.000000)
++AxisMappings=(AxisName="MoveRight",Key=A,Scale=-1.000000)
++AxisMappings=(AxisName="TurnRate",Key=MotionController_Right_Thumbstick_X,Scale=1.000000)
++AxisMappings=(AxisName="TurnRate",Key=Gamepad_RightX,Scale=1.000000)
+bAlwaysShowTouchInterface=False
+bShowConsoleOnFourFingerTap=True
+DefaultTouchInterface=None
+ConsoleKey=None
+-ConsoleKeys=Tilde
++ConsoleKeys=Tilde
 
-!AxisMappings=Empty
 
-+AxisMappings=(AxisName="MoveForward", Key=W, Scale=1.f)
-+AxisMappings=(AxisName="MoveForward", Key=S, Scale=-1.f)
-
-+AxisMappings=(AxisName="MoveRight", Key=D, Scale=1.f)
-+AxisMappings=(AxisName="MoveRight", Key=A, Scale=-1.f)
-
-!ActionMappings=Empty
-
-+ActionMappings=(ActionName="SetDestination", Key=LeftMouseButton)
\ No newline at end of file
diff --git a/Templates/TP_TopDown/Source/TP_TopDown/TP_TopDown.Build.cs b/Templates/TP_TopDown/Source/TP_TopDown/TP_TopDown.Build.cs
--- a/Templates/TP_TopDown/Source/TP_TopDown/TP_TopDown.Build.cs
+++ b/Templates/TP_TopDown/Source/TP_TopDown/TP_TopDown.Build.cs
@@ -6,6 +6,6 @@ public class TP_TopDown : ModuleRules
 {
 	public TP_TopDown(TargetInfo Target)
 	{
-		PublicDependencyModuleNames.AddRange(new string[] { "Core", "CoreUObject", "Engine", "InputCore" });
+		PublicDependencyModuleNames.AddRange(new string[] { "Core", "CoreUObject", "Engine", "InputCore", "HeadMountedDisplay" });
 	}
 }
diff --git a/Templates/TP_TopDown/Source/TP_TopDown/TP_TopDownCharacter.cpp b/Templates/TP_TopDown/Source/TP_TopDown/TP_TopDownCharacter.cpp
--- a/Templates/TP_TopDown/Source/TP_TopDown/TP_TopDownCharacter.cpp
+++ b/Templates/TP_TopDown/Source/TP_TopDown/TP_TopDownCharacter.cpp
@@ -4,6 +4,7 @@
 #include "TP_TopDownCharacter.h"
 #include "Runtime/CoreUObject/Public/UObject/ConstructorHelpers.h"
 #include "Runtime/Engine/Classes/Components/DecalComponent.h"
+#include "Kismet/HeadMountedDisplayFunctionLibrary.h"
 
 ATP_TopDownCharacter::ATP_TopDownCharacter()
 {
@@ -29,7 +30,7 @@ ATP_TopDownCharacter::ATP_TopDownCharacter()
 	CameraBoom->RelativeRotation = FRotator(-60.f, 0.f, 0.f);
 	CameraBoom->bDoCollisionTest = false; // Don't want to pull camera in when it collides with level
 
-	// Create a camera...
+										  // Create a camera...
 	TopDownCameraComponent = CreateDefaultSubobject<UCameraComponent>(TEXT("TopDownCamera"));
 	TopDownCameraComponent->SetupAttachment(CameraBoom, USpringArmComponent::SocketName);
 	TopDownCameraComponent->bUsePawnControlRotation = false; // Camera does not rotate relative to arm
@@ -54,7 +55,21 @@ void ATP_TopDownCharacter::Tick(float DeltaSeconds)
 {
 	if (CursorToWorld != nullptr)
 	{
-		if (APlayerController* PC = Cast<APlayerController>(GetController()))
+		if (UHeadMountedDisplayFunctionLibrary::IsHeadMountedDisplayEnabled())
+		{
+			if (UWorld* World = GetWorld())
+			{
+				FHitResult HitResult;
+				FCollisionQueryParams Params;
+				FVector StartLocation = TopDownCameraComponent->GetComponentLocation();
+				FVector EndLocation = TopDownCameraComponent->GetComponentRotation().Vector() * 2000.0f;
+				Params.AddIgnoredActor(this);
+				World->LineTraceSingleByChannel(HitResult, StartLocation, EndLocation, ECC_Visibility, Params);
+				FQuat SurfaceRotation = HitResult.ImpactNormal.ToOrientationRotator().Quaternion();
+				CursorToWorld->SetWorldLocationAndRotation(HitResult.Location, SurfaceRotation);
+			}
+		}
+		else if (APlayerController* PC = Cast<APlayerController>(GetController()))
 		{
 			FHitResult TraceHitResult;
 			PC->GetHitResultUnderCursor(ECC_Visibility, true, TraceHitResult);
diff --git a/Templates/TP_TopDown/Source/TP_TopDown/TP_TopDownCharacter.h b/Templates/TP_TopDown/Source/TP_TopDown/TP_TopDownCharacter.h
--- a/Templates/TP_TopDown/Source/TP_TopDown/TP_TopDownCharacter.h
+++ b/Templates/TP_TopDown/Source/TP_TopDown/TP_TopDownCharacter.h
@@ -18,6 +18,8 @@ public:
 	FORCEINLINE class UCameraComponent* GetTopDownCameraComponent() const { return TopDownCameraComponent; }
 	/** Returns CameraBoom subobject **/
 	FORCEINLINE class USpringArmComponent* GetCameraBoom() const { return CameraBoom; }
+	/** Returns CursorToWorld subobject **/
+	FORCEINLINE class UDecalComponent* GetCursorToWorld() { return CursorToWorld; }
 
 private:
 	/** Top down camera */
diff --git a/Templates/TP_TopDown/Source/TP_TopDown/TP_TopDownPlayerController.cpp b/Templates/TP_TopDown/Source/TP_TopDown/TP_TopDownPlayerController.cpp
--- a/Templates/TP_TopDown/Source/TP_TopDown/TP_TopDownPlayerController.cpp
+++ b/Templates/TP_TopDown/Source/TP_TopDown/TP_TopDownPlayerController.cpp
@@ -3,6 +3,9 @@
 #include "TP_TopDown.h"
 #include "TP_TopDownPlayerController.h"
 #include "AI/Navigation/NavigationSystem.h"
+#include "Runtime/Engine/Classes/Components/DecalComponent.h"
+#include "Kismet/HeadMountedDisplayFunctionLibrary.h"
+#include "TP_TopDownCharacter.h"
 
 ATP_TopDownPlayerController::ATP_TopDownPlayerController()
 {
@@ -32,18 +35,38 @@ void ATP_TopDownPlayerController::SetupInputComponent()
 	// support touch devices 
 	InputComponent->BindTouch(EInputEvent::IE_Pressed, this, &ATP_TopDownPlayerController::MoveToTouchLocation);
 	InputComponent->BindTouch(EInputEvent::IE_Repeat, this, &ATP_TopDownPlayerController::MoveToTouchLocation);
+
+	InputComponent->BindAction("ResetVR", IE_Pressed, this, &ATP_TopDownPlayerController::OnResetVR);
 }
 
-void ATP_TopDownPlayerController::MoveToMouseCursor()
+void ATP_TopDownPlayerController::OnResetVR()
 {
-	// Trace to see what is under the mouse cursor
-	FHitResult Hit;
-	GetHitResultUnderCursor(ECC_Visibility, false, Hit);
+	UHeadMountedDisplayFunctionLibrary::ResetOrientationAndPosition();
+}
 
-	if (Hit.bBlockingHit)
+void ATP_TopDownPlayerController::MoveToMouseCursor()
+{
+	if (UHeadMountedDisplayFunctionLibrary::IsHeadMountedDisplayEnabled())
 	{
-		// We hit something, move there
-		SetNewMoveDestination(Hit.ImpactPoint);
+		if (ATP_TopDownCharacter* MyPawn = Cast<ATP_TopDownCharacter>(GetPawn()))
+		{
+			if (MyPawn->GetCursorToWorld())
+			{
+				UNavigationSystem::SimpleMoveToLocation(this, MyPawn->GetCursorToWorld()->GetComponentLocation());
+			}
+		}
+	}
+	else
+	{
+		// Trace to see what is under the mouse cursor
+		FHitResult Hit;
+		GetHitResultUnderCursor(ECC_Visibility, false, Hit);
+
+		if (Hit.bBlockingHit)
+		{
+			// We hit something, move there
+			SetNewMoveDestination(Hit.ImpactPoint);
+		}
 	}
 }
 
diff --git a/Templates/TP_TopDown/Source/TP_TopDown/TP_TopDownPlayerController.h b/Templates/TP_TopDown/Source/TP_TopDown/TP_TopDownPlayerController.h
--- a/Templates/TP_TopDown/Source/TP_TopDown/TP_TopDownPlayerController.h
+++ b/Templates/TP_TopDown/Source/TP_TopDown/TP_TopDownPlayerController.h
@@ -20,6 +20,9 @@ protected:
 	virtual void SetupInputComponent() override;
 	// End PlayerController interface
 
+	/** Resets HMD orientation in VR. */
+	void OnResetVR();
+
 	/** Navigate player to the current mouse cursor location. */
 	void MoveToMouseCursor();
 
diff --git a/Templates/TP_Vehicle/Config/DefaultInput.ini b/Templates/TP_Vehicle/Config/DefaultInput.ini
--- a/Templates/TP_Vehicle/Config/DefaultInput.ini
+++ b/Templates/TP_Vehicle/Config/DefaultInput.ini
@@ -28,6 +28,7 @@ DoubleClickTime=0.200000
 +ActionMappings=(ActionName="Handbrake",Key=Gamepad_LeftShoulder,bShift=False,bCtrl=False,bAlt=False,bCmd=False)
 +ActionMappings=(ActionName="SwitchCamera",Key=Tab,bShift=False,bCtrl=False,bAlt=False,bCmd=False)
 +ActionMappings=(ActionName="ResetVR",Key=R,bShift=False,bCtrl=False,bAlt=False,bCmd=False)
++ActionMappings=(ActionName="ResetVR",Key=MotionController_Left_Grip1,bShift=False,bCtrl=False,bAlt=False,bCmd=False)
 +ActionMappings=(ActionName="SwitchCamera",Key=Gamepad_Special_Left,bShift=False,bCtrl=False,bAlt=False,bCmd=False)
 +AxisMappings=(AxisName="MoveForward",Key=W,Scale=1.000000)
 +AxisMappings=(AxisName="MoveForward",Key=S,Scale=-1.000000)
diff --git a/Templates/TP_VehicleAdv/Config/DefaultInput.ini b/Templates/TP_VehicleAdv/Config/DefaultInput.ini
--- a/Templates/TP_VehicleAdv/Config/DefaultInput.ini
+++ b/Templates/TP_VehicleAdv/Config/DefaultInput.ini
@@ -11,6 +11,7 @@ bUseMouseForTouch=False
 +ActionMappings=(ActionName="Handbrake",Key=Gamepad_LeftShoulder,bShift=False,bCtrl=False,bAlt=False,bCmd=False)
 +ActionMappings=(ActionName="SwitchCamera",Key=Tab,bShift=False,bCtrl=False,bAlt=False,bCmd=False)
 +ActionMappings=(ActionName="ResetVR",Key=R,bShift=False,bCtrl=True,bAlt=False,bCmd=False)
++ActionMappings=(ActionName="ResetVR",Key=MotionController_Left_Grip1,bShift=False,bCtrl=False,bAlt=False,bCmd=False)
 +ActionMappings=(ActionName="SwitchCamera",Key=Gamepad_Special_Left,bShift=False,bCtrl=False,bAlt=False,bCmd=False)
 +AxisMappings=(AxisName="MoveForward",Key=W,Scale=1.000000)
 +AxisMappings=(AxisName="MoveForward",Key=S,Scale=-1.000000)
