Top Down Template Overhaul
C++ version works but c++ behaviour is not the same yet. BP is what it will look like.
Also updated the Mannequin AnimBP to work with top down template.
Changes to material: Checker colors + grid top colors + solid grey + top dark grey.

diff --git a/Templates/TP_TopDown/Config/DefaultEngine.ini b/Templates/TP_TopDown/Config/DefaultEngine.ini
--- a/Templates/TP_TopDown/Config/DefaultEngine.ini
+++ b/Templates/TP_TopDown/Config/DefaultEngine.ini
@@ -1,7 +1,8 @@
 [/Script/EngineSettings.GameMapsSettings]
-GameDefaultMap=/Game/TopDownCPP/Maps/TopDownExampleMap
-EditorStartupMap=/Game/TopDownCPP/Maps/TopDownExampleMap
+GameDefaultMap=/Game/TopDown/Maps/TopDownMap.TopDownMap
+EditorStartupMap=/Game/TopDown/Maps/TopDownMap.TopDownMap
 GlobalDefaultGameMode="/Script/TP_TopDown.TP_TopDownGameMode"
 
 [/Script/IOSRuntimeSettings.IOSRuntimeSettings]
 MinimumiOSVersion=IOS_14
+
diff --git a/Templates/TP_TopDown/Config/DefaultInput.ini b/Templates/TP_TopDown/Config/DefaultInput.ini
--- a/Templates/TP_TopDown/Config/DefaultInput.ini
+++ b/Templates/TP_TopDown/Config/DefaultInput.ini
@@ -34,12 +34,6 @@ DefaultViewportMouseLockMode=LockOnCapture
 FOVScale=0.011110
 DoubleClickTime=0.200000
 +ActionMappings=(ActionName="SetDestination",bShift=False,bCtrl=False,bAlt=False,bCmd=False,Key=LeftMouseButton)
-+ActionMappings=(ActionName="SetDestination",bShift=False,bCtrl=False,bAlt=False,bCmd=False,Key=Gamepad_RightTrigger)
-+AxisMappings=(AxisName="MoveForward",Scale=1.000000,Key=W)
-+AxisMappings=(AxisName="MoveForward",Scale=-1.000000,Key=S)
-+AxisMappings=(AxisName="MoveRight",Scale=1.000000,Key=D)
-+AxisMappings=(AxisName="MoveRight",Scale=-1.000000,Key=A)
-+AxisMappings=(AxisName="TurnRate",Scale=1.000000,Key=Gamepad_RightX)
 DefaultPlayerInputClass=/Script/Engine.PlayerInput
 DefaultInputComponentClass=/Script/Engine.InputComponent
 DefaultTouchInterface=None
diff --git a/Templates/TP_TopDown/Config/TemplateDefs.ini b/Templates/TP_TopDown/Config/TemplateDefs.ini
--- a/Templates/TP_TopDown/Config/TemplateDefs.ini
+++ b/Templates/TP_TopDown/Config/TemplateDefs.ini
@@ -47,5 +47,6 @@ FilenameReplacements=(Extensions=("cpp","h","ini","cs"),From="%TEMPLATENAME%",To
 ReplacementsInFiles=(Extensions=("cpp","h","ini","cs"),From="%TEMPLATENAME_UPPERCASE%",To="%PROJECTNAME_UPPERCASE%",bCaseSensitive=true)
 ReplacementsInFiles=(Extensions=("cpp","h","ini","cs"),From="%TEMPLATENAME_LOWERCASE%",To="%PROJECTNAME_LOWERCASE%",bCaseSensitive=true)
 ReplacementsInFiles=(Extensions=("cpp","h","ini","cs"),From="%TEMPLATENAME%",To="%PROJECTNAME%",bCaseSensitive=false)SharedContentPacks=(MountName="TopDown",DetailLevels=("Standard"))
-SharedContentPacks=(MountName="Geometry",DetailLevels=("Standard","High"))
-SharedContentPacks=(MountName="Mannequin",DetailLevels=("Standard","High"))
\ No newline at end of file
+SharedContentPacks=(MountName="LevelPrototyping",DetailLevels=("Standard","High"))
+SharedContentPacks=(MountName="Mannequin",DetailLevels=("Standard","High"))
+EditDetailLevelPreference="High"
\ No newline at end of file
diff --git a/Templates/TP_TopDown/Config/config.ini b/Templates/TP_TopDown/Config/config.ini
--- a/Templates/TP_TopDown/Config/config.ini
+++ b/Templates/TP_TopDown/Config/config.ini
@@ -1,8 +1,4 @@
 [InputSettings]
-+AxisMappings=(AxisName="MoveForward", Key=W, Scale=1.f)
-+AxisMappings=(AxisName="MoveForward", Key=S, Scale=-1.f)
-+AxisMappings=(AxisName="MoveRight", Key=D, Scale=1.f)
-+AxisMappings=(AxisName="MoveRight", Key=A, Scale=-1.f)
 +ActionMappings=(ActionName="SetDestination", Key=LeftMouseButton)
 
 [Redirects]
diff --git a/Templates/TP_TopDown/Source/TP_TopDown/TP_TopDownCharacter.cpp b/Templates/TP_TopDown/Source/TP_TopDown/TP_TopDownCharacter.cpp
--- a/Templates/TP_TopDown/Source/TP_TopDown/TP_TopDownCharacter.cpp
+++ b/Templates/TP_TopDown/Source/TP_TopDown/TP_TopDownCharacter.cpp
@@ -43,7 +43,7 @@ ATP_TopDownCharacter::ATP_TopDownCharacter()
 	// Create a decal in the world to show the cursor's location
 	CursorToWorld = CreateDefaultSubobject<UDecalComponent>("CursorToWorld");
 	CursorToWorld->SetupAttachment(RootComponent);
-	static ConstructorHelpers::FObjectFinder<UMaterial> DecalMaterialAsset(TEXT("Material'/Game/TopDownCPP/Blueprints/M_Cursor_Decal.M_Cursor_Decal'"));
+	static ConstructorHelpers::FObjectFinder<UMaterial> DecalMaterialAsset(TEXT("Material'/Game/TopDown/Blueprints/M_Cursor_Decal.M_Cursor_Decal'"));
 	if (DecalMaterialAsset.Succeeded())
 	{
 		CursorToWorld->SetDecalMaterial(DecalMaterialAsset.Object);
diff --git a/Templates/TP_TopDown/Source/TP_TopDown/TP_TopDownGameMode.cpp b/Templates/TP_TopDown/Source/TP_TopDown/TP_TopDownGameMode.cpp
--- a/Templates/TP_TopDown/Source/TP_TopDown/TP_TopDownGameMode.cpp
+++ b/Templates/TP_TopDown/Source/TP_TopDown/TP_TopDownGameMode.cpp
@@ -11,7 +11,7 @@ ATP_TopDownGameMode::ATP_TopDownGameMode()
 	PlayerControllerClass = ATP_TopDownPlayerController::StaticClass();
 
 	// set default pawn class to our Blueprinted character
-	static ConstructorHelpers::FClassFinder<APawn> PlayerPawnBPClass(TEXT("/Game/TopDownCPP/Blueprints/TopDownCharacter"));
+	static ConstructorHelpers::FClassFinder<APawn> PlayerPawnBPClass(TEXT("/Game/TopDown/Blueprints/BP_TopDownCharacter"));
 	if (PlayerPawnBPClass.Class != nullptr)
 	{
 		DefaultPawnClass = PlayerPawnBPClass.Class;
diff --git a/Templates/TP_TopDown/manifest.json b/Templates/TP_TopDown/manifest.json
--- a/Templates/TP_TopDown/manifest.json
+++ b/Templates/TP_TopDown/manifest.json
@@ -54,7 +54,7 @@
 	],
 	"ClassTypes": "GameMode, Character, CameraComponent, SpringArmComponent, PlayerController",
 	"Category": "CodeFeature",
-	"FocusAsset": "/Game/TopDownCPP/TopDownOverview.TopDownOverview",
+	"FocusAsset": "/Game/TopDown/Blueprints",
 	"Thumbnail": "TP_TopDown.png",
 	"Screenshots":
 	[
@@ -63,18 +63,16 @@
 	"AdditionalFeaturePacks":
 	[
 		{
-			"MountName": "Geometry",
+			"MountName": "LevelPrototyping",
 			"DetailLevels":
 			[
-				"Standard",
 				"High"
 			]
-		},
+		},				
 		{
 			"MountName": "Mannequin",
 			"DetailLevels":
 			[
-				"Standard",
 				"High"
 			]
 		}
diff --git a/Templates/TP_TopDownBP/Config/DefaultEngine.ini b/Templates/TP_TopDownBP/Config/DefaultEngine.ini
--- a/Templates/TP_TopDownBP/Config/DefaultEngine.ini
+++ b/Templates/TP_TopDownBP/Config/DefaultEngine.ini
@@ -2,13 +2,13 @@
 GameName=TP_TopDownBP
 
 [/Script/EngineSettings.GameMapsSettings]
-EditorStartupMap=/Game/TopDownBP/Maps/TopDownExampleMap
-GameDefaultMap=/Game/TopDownBP/Maps/TopDownExampleMap
+EditorStartupMap=/Game/TopDown/Maps/TopDownMap.TopDownMap
+GameDefaultMap=/Game/TopDown/Maps/TopDownMap.TopDownMap
 TransitionMap=
 bUseSplitscreen=True
 TwoPlayerSplitscreenLayout=Horizontal
 ThreePlayerSplitscreenLayout=FavorTop
-GlobalDefaultGameMode=/Game/TopDownBP/Blueprints/TopDownGameMode.TopDownGameMode_C
+GlobalDefaultGameMode=/Game/TopDown/Blueprints/BP_TopDownGameMode.BP_TopDownGameMode_C
 GlobalDefaultServerGameMode=None
 
 [/Script/NavigationSystem.RecastNavMesh]
diff --git a/Templates/TP_TopDownBP/Config/DefaultInput.ini b/Templates/TP_TopDownBP/Config/DefaultInput.ini
--- a/Templates/TP_TopDownBP/Config/DefaultInput.ini
+++ b/Templates/TP_TopDownBP/Config/DefaultInput.ini
@@ -34,15 +34,10 @@ DefaultViewportMouseLockMode=LockOnCapture
 FOVScale=0.011110
 DoubleClickTime=0.200000
 +ActionMappings=(ActionName="SetDestination",bShift=False,bCtrl=False,bAlt=False,bCmd=False,Key=LeftMouseButton)
-+ActionMappings=(ActionName="SetDestination",bShift=False,bCtrl=False,bAlt=False,bCmd=False,Key=Gamepad_RightTrigger)
-+AxisMappings=(AxisName="MoveForward",Scale=1.000000,Key=W)
-+AxisMappings=(AxisName="MoveForward",Scale=-1.000000,Key=S)
-+AxisMappings=(AxisName="MoveRight",Scale=1.000000,Key=D)
-+AxisMappings=(AxisName="MoveRight",Scale=-1.000000,Key=A)
-+AxisMappings=(AxisName="TurnRate",Scale=1.000000,Key=Gamepad_RightX)
 DefaultPlayerInputClass=/Script/Engine.PlayerInput
 DefaultInputComponentClass=/Script/Engine.InputComponent
 DefaultTouchInterface=None
 -ConsoleKeys=Tilde
 +ConsoleKeys=Tilde
++ConsoleKeys=Caret
 
diff --git a/Templates/TP_TopDownBP/Config/TemplateDefs.ini b/Templates/TP_TopDownBP/Config/TemplateDefs.ini
--- a/Templates/TP_TopDownBP/Config/TemplateDefs.ini
+++ b/Templates/TP_TopDownBP/Config/TemplateDefs.ini
@@ -42,6 +42,6 @@ FilenameReplacements=(Extensions=("cpp","h","ini","cs"),From="%TEMPLATENAME%",To
 ReplacementsInFiles=(Extensions=("cpp","h","ini","cs"),From="%TEMPLATENAME_UPPERCASE%",To="%PROJECTNAME_UPPERCASE%",bCaseSensitive=true)
 ReplacementsInFiles=(Extensions=("cpp","h","ini","cs"),From="%TEMPLATENAME_LOWERCASE%",To="%PROJECTNAME_LOWERCASE%",bCaseSensitive=true)
 ReplacementsInFiles=(Extensions=("cpp","h","ini","cs"),From="%TEMPLATENAME%",To="%PROJECTNAME%",bCaseSensitive=false)
-SharedContentPacks=(MountName="Geometry",DetailLevels=("Standard","High"))
+SharedContentPacks=(MountName="LevelPrototyping",DetailLevels=("Standard","High"))
 SharedContentPacks=(MountName="Mannequin",DetailLevels=("Standard","High"))
 EditDetailLevelPreference="High"
diff --git a/Templates/TP_TopDownBP/Config/config.ini b/Templates/TP_TopDownBP/Config/config.ini
--- a/Templates/TP_TopDownBP/Config/config.ini
+++ b/Templates/TP_TopDownBP/Config/config.ini
@@ -1,8 +1,4 @@
 [InputSettings]
-+AxisMappings=(AxisName="MoveForward", Key=W, Scale=1.f)
-+AxisMappings=(AxisName="MoveForward", Key=S, Scale=-1.f)
-+AxisMappings=(AxisName="MoveRight", Key=D, Scale=1.f)
-+AxisMappings=(AxisName="MoveRight", Key=A, Scale=-1.f)
 +ActionMappings=(ActionName="SetDestination", Key=LeftMouseButton)
 
 [Redirects]
diff --git a/Templates/TP_TopDownBP/TP_TopDownBP.uproject b/Templates/TP_TopDownBP/TP_TopDownBP.uproject
--- a/Templates/TP_TopDownBP/TP_TopDownBP.uproject
+++ b/Templates/TP_TopDownBP/TP_TopDownBP.uproject
@@ -10,6 +10,52 @@
 			"TargetAllowList": [
 				"Editor"
 			]
+		},
+		{
+			"Name": "OculusVR",
+			"Enabled": false,
+			"SupportedTargetPlatforms": [
+				"Win64",
+				"Android"
+			]
+		},
+		{
+			"Name": "OpenXREyeTracker",
+			"Enabled": false,
+			"SupportedTargetPlatforms": [
+				"Win64",
+				"Linux",
+				"HoloLens",
+				"Android"
+			]
+		},
+		{
+			"Name": "OpenXRHandTracking",
+			"Enabled": false,
+			"SupportedTargetPlatforms": [
+				"Win64",
+				"Linux",
+				"HoloLens",
+				"Android"
+			]
+		},
+		{
+			"Name": "OpenXR",
+			"Enabled": false,
+			"SupportedTargetPlatforms": [
+				"Win64",
+				"Linux",
+				"Android",
+				"HoloLens"
+			]
+		},
+		{
+			"Name": "SteamVR",
+			"Enabled": false,
+			"SupportedTargetPlatforms": [
+				"Win64",
+				"Linux"
+			]
 		}
-	]	
+	]
 }
\ No newline at end of file
diff --git a/Templates/TP_TopDownBP/manifest.json b/Templates/TP_TopDownBP/manifest.json
--- a/Templates/TP_TopDownBP/manifest.json
+++ b/Templates/TP_TopDownBP/manifest.json
@@ -54,7 +54,7 @@
 	],
 	"ClassTypes": "GameMode, Character, CameraComponent, SpringArmComponent, PlayerController",
 	"Category": "BlueprintFeature",
-	"FocusAsset": "/Game/TopDownBP/TopDownOverview.TopDownOverview",
+	"FocusAsset": "/Game/TopDown/Blueprints",
 	"Thumbnail": "TP_TopDownBP.png",
 	"Screenshots":
 	[
@@ -63,10 +63,9 @@
 	"AdditionalFeaturePacks":
 	[
 		{
-			"MountName": "Geometry",
+			"MountName": "LevelPrototyping",
 			"DetailLevels":
-			[
-				"Standard",
+			[			
 				"High"
 			]
 		},
@@ -74,7 +73,6 @@
 			"MountName": "Mannequin",
 			"DetailLevels":
 			[
-				"Standard",
 				"High"
 			]
 		}
