Templates
- Removed ActiveGameRedirects that caused C++ code added to FirstPersonBP and VehicleBP templates to not resolve package names correctly.
- Added manifest.json and contents.txt to the ignored files on new project creation for FirstPersonBP
- Removed old CoreRedirects on TopDown

#rnx

diff --git a/Templates/TP_FirstPersonBP/Config/DefaultEngine.ini b/Templates/TP_FirstPersonBP/Config/DefaultEngine.ini
--- a/Templates/TP_FirstPersonBP/Config/DefaultEngine.ini
+++ b/Templates/TP_FirstPersonBP/Config/DefaultEngine.ini
@@ -54,13 +54,10 @@ DefaultGraphicsPerformance=Maximum
 AppliedDefaultGraphicsPerformance=Maximum
 
 [/Script/Engine.Engine]
-+ActiveGameNameRedirects=(OldGameName="TP_FirstPersonBP",NewGameName="/Script/MyProject")
-+ActiveGameNameRedirects=(OldGameName="/Script/TP_FirstPersonBP",NewGameName="/Script/MyProject")
 NearClipPlane=5.000000
 
 [/Script/AIModule.AISystem]
 bForgetStaleActors=True
 
 [/Script/Engine.UserInterfaceSettings]
-UIScaleRule=ShortestSide
-
+UIScaleRule=ShortestSide
\ No newline at end of file
diff --git a/Templates/TP_FirstPersonBP/Config/TemplateDefs.ini b/Templates/TP_FirstPersonBP/Config/TemplateDefs.ini
--- a/Templates/TP_FirstPersonBP/Config/TemplateDefs.ini
+++ b/Templates/TP_FirstPersonBP/Config/TemplateDefs.ini
@@ -21,6 +21,8 @@ FilesToIgnore="%TEMPLATENAME%.uproject"
 FilesToIgnore="%TEMPLATENAME%.png"
 FilesToIgnore="Config/TemplateDefs.ini"
 FilesToIgnore="Config/config.ini"
+FilesToIgnore="Manifest.json"
+FilesToIgnore="contents.txt"
 
 FolderRenames=(From="Source/%TEMPLATENAME%",To="Source/%PROJECTNAME%")
 FolderRenames=(From="Source/%TEMPLATENAME%Editor",To="Source/%PROJECTNAME%Editor")
diff --git a/Templates/TP_TopDown/Config/DefaultEngine.ini b/Templates/TP_TopDown/Config/DefaultEngine.ini
--- a/Templates/TP_TopDown/Config/DefaultEngine.ini
+++ b/Templates/TP_TopDown/Config/DefaultEngine.ini
@@ -40,29 +40,5 @@ SoundCueCookQualityIndex=-1
 -TargetedRHIs=SF_VULKAN_SM5
 +TargetedRHIs=SF_VULKAN_SM6
 
-[CoreRedirects]
-+ClassRedirects=(OldName="/Script/TopDownCPP.StrategyGameMode",NewName="/Script/TP_TopDown.StrategyGameMode")
-+ClassRedirects=(OldName="/Script/TopDownCPP.StrategyHUD",NewName="/Script/TP_TopDown.StrategyHUD")
-+ClassRedirects=(OldName="/Script/TopDownCPP.StrategyPawn",NewName="/Script/TP_TopDown.StrategyPawn")
-+ClassRedirects=(OldName="/Script/TopDownCPP.StrategyPlayerController",NewName="/Script/TP_TopDown.StrategyPlayerController")
-+ClassRedirects=(OldName="/Script/TopDownCPP.StrategyUI",NewName="/Script/TP_TopDown.StrategyUI")
-+ClassRedirects=(OldName="/Script/TopDownCPP.StrategyNPC",NewName="/Script/TP_TopDown.StrategyUnit")
-
-+ClassRedirects=(OldName="/Script/TopDownCPP.TwinStickAIController",NewName="/Script/TP_TopDown.TwinStickAIController")
-+ClassRedirects=(OldName="/Script/TopDownCPP.TwinStickNPC",NewName="/Script/TP_TopDown.TwinStickNPC")
-+ClassRedirects=(OldName="/Script/TopDownCPP.TwinStickAoEAttack",NewName="/Script/TP_TopDown.TwinStickAoEAttack")
-+ClassRedirects=(OldName="/Script/TopDownCPP.TwinStickCharacter",NewName="/Script/TP_TopDown.TwinStickCharacter")
-+ClassRedirects=(OldName="/Script/TopDownCPP.TwinStickGameMode",NewName="/Script/TP_TopDown.TwinStickGameMode")
-+ClassRedirects=(OldName="/Script/TopDownCPP.TwinStickNPCDestruction",NewName="/Script/TP_TopDown.TwinStickNPCDestruction")
-+ClassRedirects=(OldName="/Script/TopDownCPP.TwinStickPickup",NewName="/Script/TP_TopDown.TwinStickPickup")
-+ClassRedirects=(OldName="/Script/TopDownCPP.TwinStickPlayerController",NewName="/Script/TP_TopDown.TwinStickPlayerController")
-+ClassRedirects=(OldName="/Script/TopDownCPP.TwinStickProjectile",NewName="/Script/TP_TopDown.TwinStickProjectile")
-+ClassRedirects=(OldName="/Script/TopDownCPP.TwinStickSpawner",NewName="/Script/TP_TopDown.TwinStickSpawner")
-+ClassRedirects=(OldName="/Script/TopDownCPP.TwinStickUI",NewName="/Script/TP_TopDown.TwinStickUI")
-
-;+StructRedirects=(OldName="MyStruct",NewName="MyNewStruct")
-;+StructRedirects=(OldName="MyStruct",NewName="MyNewStruct")
-
 [/Script/NavigationSystem.RecastNavMesh]
-RuntimeGeneration=Dynamic
-
+RuntimeGeneration=Dynamic
\ No newline at end of file
diff --git a/Templates/TP_VehicleAdvBP/Config/DefaultEngine.ini b/Templates/TP_VehicleAdvBP/Config/DefaultEngine.ini
--- a/Templates/TP_VehicleAdvBP/Config/DefaultEngine.ini
+++ b/Templates/TP_VehicleAdvBP/Config/DefaultEngine.ini
@@ -77,23 +77,4 @@ SoundCueCookQualityIndex=-1
 TargetedHardwareClass=Desktop
 AppliedTargetedHardwareClass=Desktop
 DefaultGraphicsPerformance=Maximum
-AppliedDefaultGraphicsPerformance=Maximum
-
-[/Script/Engine.Engine]
-+ActiveGameNameRedirects=(OldGameName="TP_VehicleAdvBP",NewGameName="/Script/VehicleTemplateNew")
-+ActiveGameNameRedirects=(OldGameName="/Script/TP_VehicleAdvBP",NewGameName="/Script/VehicleTemplateNew")
-
-[/Script/AndroidFileServerEditor.AndroidFileServerRuntimeSettings]
-bEnablePlugin=True
-bAllowNetworkConnection=True
-SecurityToken=1EF3860741EAFFC42C80FFA0493E9638
-bIncludeInShipping=False
-bAllowExternalStartInShipping=False
-bCompileAFSProject=False
-bUseCompression=False
-bLogFiles=False
-bReportStats=False
-ConnectionType=USBOnly
-bUseManualIPAddress=False
-ManualIPAddress=
-
+AppliedDefaultGraphicsPerformance=Maximum
\ No newline at end of file
